Remove PLATFORM_HTML5 reference
#rnx

diff --git a/Engine/Plugins/Online/WebAuth/Source/Public/PlatformWebAuth.h b/Engine/Plugins/Online/WebAuth/Source/Public/PlatformWebAuth.h
--- a/Engine/Plugins/Online/WebAuth/Source/Public/PlatformWebAuth.h
+++ b/Engine/Plugins/Online/WebAuth/Source/Public/PlatformWebAuth.h
@@ -15,8 +15,6 @@
 #include "IOS/IOSPlatformWebAuth.h"
 #elif PLATFORM_ANDROID
 #include "Android/AndroidPlatformWebAuth.h"
-#elif PLATFORM_HTML5
-#include "NullPlatformWebAuth.h"
 #elif PLATFORM_UNIX
 #include "NullPlatformWebAuth.h"
 #elif PLATFORM_SWITCH
