Fix non-unity
#rnx

diff --git a/Engine/Plugins/Online/WebAuth/Source/Public/WebAuthModule.h b/Engine/Plugins/Online/WebAuth/Source/Public/WebAuthModule.h
--- a/Engine/Plugins/Online/WebAuth/Source/Public/WebAuthModule.h
+++ b/Engine/Plugins/Online/WebAuth/Source/Public/WebAuthModule.h
@@ -4,6 +4,7 @@
 
 // Module dependencies
 #include "CoreMinimal.h"
+#include "Misc/CoreMisc.h"
 #include "Modules/ModuleInterface.h"
 
 // Module includes
