Limit whitelisted platforms for WebAuth plugin
Only include for explicitly listed for platforms, and exclude TVOS variants of IOS

diff --git a/Engine/Plugins/Online/WebAuth/Source/Private/IOS/IOSPlatformWebAuth.cpp b/Engine/Plugins/Online/WebAuth/Source/Private/IOS/IOSPlatformWebAuth.cpp
--- a/Engine/Plugins/Online/WebAuth/Source/Private/IOS/IOSPlatformWebAuth.cpp
+++ b/Engine/Plugins/Online/WebAuth/Source/Private/IOS/IOSPlatformWebAuth.cpp
@@ -1,4 +1,5 @@
 // Copyright 1998-2019 Epic Games, Inc. All Rights Reserved.
+#if  PLATFORM_IOS && !PLATFORM_TVOS
 
 #include "IOS/IOSPlatformWebAuth.h"
 
@@ -253,3 +254,4 @@ IWebAuth* FIOSPlatformWebAuth::CreatePlatformWebAuth()
 	return new FIOSWebAuth();
 }
 
+#endif // PLATFORM_IOS && !PLATFORM_TVOS
\ No newline at end of file
diff --git a/Engine/Plugins/Online/WebAuth/Source/Public/PlatformWebAuth.h b/Engine/Plugins/Online/WebAuth/Source/Public/PlatformWebAuth.h
--- a/Engine/Plugins/Online/WebAuth/Source/Public/PlatformWebAuth.h
+++ b/Engine/Plugins/Online/WebAuth/Source/Public/PlatformWebAuth.h
@@ -3,23 +3,11 @@
 
 #include "CoreMinimal.h"
 
-#if PLATFORM_WINDOWS
-#include "NullPlatformWebAuth.h"
-#elif PLATFORM_PS4
-#include "NullPlatformWebAuth.h"
-#elif PLATFORM_XBOXONE
-#include "NullPlatformWebAuth.h"
-#elif PLATFORM_MAC
-#include "NullPlatformWebAuth.h"
-#elif PLATFORM_IOS
+#if  PLATFORM_IOS && !PLATFORM_TVOS
 #include "IOS/IOSPlatformWebAuth.h"
 #elif PLATFORM_ANDROID
 #include "Android/AndroidPlatformWebAuth.h"
-#elif PLATFORM_UNIX
-#include "NullPlatformWebAuth.h"
-#elif PLATFORM_SWITCH
-#include "NullPlatformWebAuth.h"
 #else
-#error Unknown platform
+#include "NullPlatformWebAuth.h"
 #endif
 
diff --git a/Engine/Plugins/Online/WebAuth/WebAuth.uplugin b/Engine/Plugins/Online/WebAuth/WebAuth.uplugin
--- a/Engine/Plugins/Online/WebAuth/WebAuth.uplugin
+++ b/Engine/Plugins/Online/WebAuth/WebAuth.uplugin
@@ -14,7 +14,13 @@
 		{
 			"Name": "WebAuth",
 			"Type": "Runtime",
-            "LoadingPhase" : "Default"
+			"LoadingPhase" : "Default",
+			"WhitelistPlatforms" :
+			[
+				"Win64",
+				"iOS",
+				"Android"
+			]
 		}
 	]
 }
