Fix iOS non-unity compile
#rnx

diff --git a/Engine/Plugins/Online/WebAuth/Source/Private/IOS/IOSPlatformWebAuth.cpp b/Engine/Plugins/Online/WebAuth/Source/Private/IOS/IOSPlatformWebAuth.cpp
--- a/Engine/Plugins/Online/WebAuth/Source/Private/IOS/IOSPlatformWebAuth.cpp
+++ b/Engine/Plugins/Online/WebAuth/Source/Private/IOS/IOSPlatformWebAuth.cpp
@@ -1,7 +1,7 @@
 // Copyright 1998-2019 Epic Games, Inc. All Rights Reserved.
-#if  PLATFORM_IOS && !PLATFORM_TVOS
 
-#include "IOS/IOSPlatformWebAuth.h"
+#include "PlatformWebAuth.h"
+#if PLATFORM_IOS && !PLATFORM_TVOS
 
 #import <Foundation/Foundation.h>
 #ifdef WEBAUTH_PLATFORM_IOS_12
