Fix iOS compile caused by attempt to fix non-unity
#rnx

diff --git a/Engine/Plugins/Online/WebAuth/Source/Private/IOS/IOSPlatformWebAuth.cpp b/Engine/Plugins/Online/WebAuth/Source/Private/IOS/IOSPlatformWebAuth.cpp
--- a/Engine/Plugins/Online/WebAuth/Source/Private/IOS/IOSPlatformWebAuth.cpp
+++ b/Engine/Plugins/Online/WebAuth/Source/Private/IOS/IOSPlatformWebAuth.cpp
@@ -1,6 +1,6 @@
 // Copyright 1998-2019 Epic Games, Inc. All Rights Reserved.
 
-#include "PlatformWebAuth.h"
+#include "IOS/IOSPlatformWebAuth.h"
 #if PLATFORM_IOS && !PLATFORM_TVOS
 
 #import <Foundation/Foundation.h>
diff --git a/Engine/Plugins/Online/WebAuth/Source/Public/IOS/IOSPlatformWebAuth.h b/Engine/Plugins/Online/WebAuth/Source/Public/IOS/IOSPlatformWebAuth.h
--- a/Engine/Plugins/Online/WebAuth/Source/Public/IOS/IOSPlatformWebAuth.h
+++ b/Engine/Plugins/Online/WebAuth/Source/Public/IOS/IOSPlatformWebAuth.h
@@ -2,6 +2,10 @@
 
 #pragma once
 
+#include "CoreMinimal.h"
+
+#if PLATFORM_IOS && !PLATFORM_TVOS
+
 #include "GenericPlatform/GenericPlatformWebAuth.h"
 #include "WebAuth.h"
 
@@ -49,3 +53,4 @@ public:
 
 typedef FIOSPlatformWebAuth FPlatformWebAuth;
 
+#endif
\ No newline at end of file
