WebAuth IOS 11 fix for AuthenticationServices not found error

diff --git a/Engine/Plugins/Online/WebAuth/Source/WebAuth.Build.cs b/Engine/Plugins/Online/WebAuth/Source/WebAuth.Build.cs
--- a/Engine/Plugins/Online/WebAuth/Source/WebAuth.Build.cs
+++ b/Engine/Plugins/Online/WebAuth/Source/WebAuth.Build.cs
@@ -19,11 +19,9 @@ public class WebAuth : ModuleRules
 
 		if (Target.Platform == UnrealTargetPlatform.IOS)
 		{
-			// Building for iOS 12+ Only
-			PublicFrameworks.Add("AuthenticationServices");
-			PublicDefinitions.Add("WEBAUTH_PLATFORM_IOS_12");
-			// << Building for iOS 12+ Only
-
+			// << Add these lines to your project's build.cs when building for iOS 12+
+			// PublicFrameworks.Add("AuthenticationServices");
+			// PublicDefinitions.Add("WEBAUTH_PLATFORM_IOS_12");
 
 			// iOS 11+ legacy support
 			PublicFrameworks.Add("SafariServices");
