mark WebAuth as unavailable for servers
#rnx

diff --git a/Engine/Plugins/Online/WebAuth/WebAuth.uplugin b/Engine/Plugins/Online/WebAuth/WebAuth.uplugin
--- a/Engine/Plugins/Online/WebAuth/WebAuth.uplugin
+++ b/Engine/Plugins/Online/WebAuth/WebAuth.uplugin
@@ -20,7 +20,8 @@
 				"Win64",
 				"iOS",
 				"Android"
-			]
+			],
+			"BlacklistTargets": [ "Server" ]
 		}
 	]
 }
