Attempting to fix iOS compile issue found by CIS.

diff --git a/Engine/Plugins/Online/WebAuth/Source/WebAuth.Build.cs b/Engine/Plugins/Online/WebAuth/Source/WebAuth.Build.cs
--- a/Engine/Plugins/Online/WebAuth/Source/WebAuth.Build.cs
+++ b/Engine/Plugins/Online/WebAuth/Source/WebAuth.Build.cs
@@ -20,6 +20,7 @@ public class WebAuth : ModuleRules
 
 		if (Target.Platform == UnrealTargetPlatform.IOS)
 		{
+			PrivateDependencyModuleNames.Add("ApplicationCore");
 			PublicFrameworks.Add("AuthenticationServices");
 		}
 
