Add support to web auth to complete using an external browser on iOS

diff --git a/Engine/Plugins/Online/WebAuth/Source/Private/IOS/IOSPlatformWebAuth.cpp b/Engine/Plugins/Online/WebAuth/Source/Private/IOS/IOSPlatformWebAuth.cpp
--- a/Engine/Plugins/Online/WebAuth/Source/Private/IOS/IOSPlatformWebAuth.cpp
+++ b/Engine/Plugins/Online/WebAuth/Source/Private/IOS/IOSPlatformWebAuth.cpp
@@ -8,6 +8,7 @@
 
 #import <WebKit/WebKit.h>
 
+#include "IOS/IOSAsyncTask.h"
 #include "IOS/IOSAppDelegate.h"
 
 @interface PresentationContext : NSObject <ASWebAuthenticationPresentationContextProviding>
@@ -31,7 +32,6 @@
 static NSMutableDictionary* MakeSearchDictionary(NSString *EnvironmentName);
 static PresentationContext* PresentationContextProvider = nullptr;
 
-
 // Returns true if the request is made
 bool FIOSWebAuth::AuthSessionWithURL(const FString &UrlStr, const FString &SchemeStr, const FWebAuthSessionCompleteDelegate& Delegate)
 {
@@ -48,27 +48,50 @@ bool FIOSWebAuth::AuthSessionWithURL(const FString &UrlStr, const FString &Schem
 
     bSessionInProgress = true;
     AuthSessionCompleteDelegate = Delegate;
+	FIOSCoreDelegates::OnOpenURL.Remove(OpenUrlHandle);
 
-    id SavedAuthSession = [[ASWebAuthenticationSession alloc] initWithURL:Url callbackURLScheme:Scheme completionHandler:^(NSURL * _Nullable callbackURL, NSError * _Nullable error)
+	ASWebAuthenticationSession* SavedAuthSession = [[ASWebAuthenticationSession alloc] initWithURL:Url callbackURLScheme:Scheme completionHandler:^(NSURL * _Nullable callbackURL, NSError * _Nullable error)
     {
-        // Response received
-        if (callbackURL != nil)
-        {
-            const char *StrCallbackURL = [callbackURL.absoluteString UTF8String];
-            AuthSessionCompleteDelegate.ExecuteIfBound(FString(StrCallbackURL), true);
-        }
-        // Empty response
-        else
-        {
-            AuthSessionCompleteDelegate.ExecuteIfBound(FString(), false);
-        }
-        AuthSessionCompleteDelegate = nullptr;
+		[FIOSAsyncTask CreateTaskWithBlock : ^ bool(void)
+		{
+			FIOSCoreDelegates::OnOpenURL.Remove(OpenUrlHandle);
+			// Response received
+			if (callbackURL != nil)
+			{
+				const char *StrCallbackURL = [callbackURL.absoluteString UTF8String];
+				AuthSessionCompleteDelegate.ExecuteIfBound(UTF8_TO_TCHAR(StrCallbackURL), true);
+			}
+			// Empty response
+			else
+			{
+				AuthSessionCompleteDelegate.ExecuteIfBound(FString(), false);
+			}
+			AuthSessionCompleteDelegate = nullptr;
+			return true;
+		}];
     }];
 
     check(PresentationContextProvider);
-    ((ASWebAuthenticationSession*)SavedAuthSession).presentationContextProvider = PresentationContextProvider;
+    SavedAuthSession.presentationContextProvider = PresentationContextProvider;
+
+	OpenUrlHandle = FIOSCoreDelegates::OnOpenURL.AddLambda([this, SchemeStr, SavedAuthSession](UIApplication* application, NSURL* url, NSString* sourceApplication, id annotation)
+	{
+		if (FString(url.scheme) == SchemeStr)
+		{
+			[SavedAuthSession cancel];
+
+			[FIOSAsyncTask CreateTaskWithBlock : ^ bool(void)
+			{
+				FIOSCoreDelegates::OnOpenURL.Remove(OpenUrlHandle);
+				const char* StrCallbackURL = [url.absoluteString UTF8String];
+				AuthSessionCompleteDelegate.ExecuteIfBound(UTF8_TO_TCHAR(StrCallbackURL), true);
+				AuthSessionCompleteDelegate = nullptr;
+				return true;
+			}];
+		}
+	});
 
-    [(ASWebAuthenticationSession*)SavedAuthSession start];
+	[SavedAuthSession start];
 
 	return bSessionInProgress;
 }
@@ -212,6 +235,7 @@ FIOSWebAuth::FIOSWebAuth()
 
 FIOSWebAuth::~FIOSWebAuth()
 {
+	FIOSCoreDelegates::OnOpenURL.Remove(OpenUrlHandle);
 	if (PresentationContextProvider != nil)
 	{
 		[PresentationContextProvider release];
diff --git a/Engine/Plugins/Online/WebAuth/Source/Public/IOS/IOSPlatformWebAuth.h b/Engine/Plugins/Online/WebAuth/Source/Public/IOS/IOSPlatformWebAuth.h
--- a/Engine/Plugins/Online/WebAuth/Source/Public/IOS/IOSPlatformWebAuth.h
+++ b/Engine/Plugins/Online/WebAuth/Source/Public/IOS/IOSPlatformWebAuth.h
@@ -36,6 +36,8 @@ public:
 	virtual ~FIOSWebAuth();
 
 private:
+	/** Delegate handle to get auth response from external browser */
+	FDelegateHandle OpenUrlHandle;
 	/** Delegate that will get called once an authentication session completes or for an error condition */
 	FWebAuthSessionCompleteDelegate AuthSessionCompleteDelegate;
 };
@@ -53,4 +55,4 @@ public:
 
 typedef FIOSPlatformWebAuth FPlatformWebAuth;
 
-#endif
\ No newline at end of file
+#endif
