Manual merge

Fix misc memory leaks

#rnx

diff --git a/Engine/Plugins/Online/WebAuth/Source/Private/IOS/IOSPlatformWebAuth.cpp b/Engine/Plugins/Online/WebAuth/Source/Private/IOS/IOSPlatformWebAuth.cpp
--- a/Engine/Plugins/Online/WebAuth/Source/Private/IOS/IOSPlatformWebAuth.cpp
+++ b/Engine/Plugins/Online/WebAuth/Source/Private/IOS/IOSPlatformWebAuth.cpp
@@ -92,6 +92,7 @@ bool FIOSWebAuth::AuthSessionWithURL(const FString &UrlStr, const FString &Schem
 	});
 
 	[SavedAuthSession start];
+    [SavedAuthSession release];
 
 	return bSessionInProgress;
 }
@@ -109,7 +110,7 @@ NSMutableDictionary* MakeSearchDictionary(NSString *EnvironmentName)
 	[SearchDictionary setObject:EncodedIdentifier forKey:(id)kSecAttrAccount];
 	[SearchDictionary setObject:ServiceName forKey:(id)kSecAttrService];
 
-	return SearchDictionary;
+	return [SearchDictionary autorelease];
 }
 
 bool FIOSWebAuth::SaveCredentials(const FString& IdStr, const FString& TokenStr, const FString& EnvironmentNameStr)
diff --git a/Engine/Plugins/Runtime/Firebase/Source/IOS/notifications/EpicFirebaseIOSNotifications.cpp b/Engine/Plugins/Runtime/Firebase/Source/IOS/notifications/EpicFirebaseIOSNotifications.cpp
--- a/Engine/Plugins/Runtime/Firebase/Source/IOS/notifications/EpicFirebaseIOSNotifications.cpp
+++ b/Engine/Plugins/Runtime/Firebase/Source/IOS/notifications/EpicFirebaseIOSNotifications.cpp
@@ -109,6 +109,7 @@ NSString* KEY_FIREBASE_TOKEN = @"firebasetoken";
         }];
     });
     dispatch_semaphore_wait(updateTokenSemaphore, dispatch_time(DISPATCH_TIME_NOW, Timeout));
+    dispatch_release(updateTokenSemaphore);
 }
 @end
 
diff --git a/Engine/Source/Runtime/Core/Private/Async/ParkingLot.cpp b/Engine/Source/Runtime/Core/Private/Async/ParkingLot.cpp
--- a/Engine/Source/Runtime/Core/Private/Async/ParkingLot.cpp
+++ b/Engine/Source/Runtime/Core/Private/Async/ParkingLot.cpp
@@ -108,6 +108,8 @@ public:
 
 	static TRefCountPtr<FThread> Get()
 	{
+        // This static will appear as memory leaks when comparing AnsiMalloc and AnsiFree
+        // But since this only grows as new threads created, it won't grow indefinitely
 		static thread_local FThreadLocalData ThreadLocalData;
 		FThreadLocalData& LocalThreadLocalData = ThreadLocalData;
 		if (LIKELY(LocalThreadLocalData.Thread))
diff --git a/Engine/Source/Runtime/Core/Private/IOS/IOSPlatformCrashContext.cpp b/Engine/Source/Runtime/Core/Private/IOS/IOSPlatformCrashContext.cpp
--- a/Engine/Source/Runtime/Core/Private/IOS/IOSPlatformCrashContext.cpp
+++ b/Engine/Source/Runtime/Core/Private/IOS/IOSPlatformCrashContext.cpp
@@ -83,6 +83,7 @@ void FIOSCrashContext::ConvertMinidump(char const* OutputPath, char const* Input
 		{
 			NSLog(@"****UE4 %@", [Error localizedDescription]);
 		}
+        [CrashLog release];
 	}
 	else
 	{
diff --git a/Engine/Source/Runtime/Core/Private/IOS/IOSPlatformMisc.cpp b/Engine/Source/Runtime/Core/Private/IOS/IOSPlatformMisc.cpp
--- a/Engine/Source/Runtime/Core/Private/IOS/IOSPlatformMisc.cpp
+++ b/Engine/Source/Runtime/Core/Private/IOS/IOSPlatformMisc.cpp
@@ -1317,6 +1317,7 @@ FIOSPlatformMisc::EIOSAuthNotificationStatus FIOSPlatformMisc::GetNotificationAu
 
 	// wait for a result, but timeout after 1s
 	dispatch_semaphore_wait(Semaphore, dispatch_time(DISPATCH_TIME_NOW, NSEC_PER_SEC));
+    dispatch_release(Semaphore);
 	return CurrentAuthStatus;
 }
 
diff --git a/Engine/Source/Runtime/Core/Public/Mac/MacSemaphore.h b/Engine/Source/Runtime/Core/Public/Mac/MacSemaphore.h
--- a/Engine/Source/Runtime/Core/Public/Mac/MacSemaphore.h
+++ b/Engine/Source/Runtime/Core/Public/Mac/MacSemaphore.h
@@ -18,6 +18,11 @@ public:
 	{
 		checkfSlow(InitialCount >= 0, TEXT("Semaphore's initial count must be non negative value: %d"), InitialCount);
 	}
+    
+    virtual ~FMacSemaphore()
+    {
+        dispatch_release(Semaphore);
+    }
 
 	void Acquire()
 	{
@@ -45,4 +50,4 @@ private:
 	dispatch_semaphore_t Semaphore;
 };
 
-typedef FMacSemaphore FSemaphore;
\ No newline at end of file
+typedef FMacSemaphore FSemaphore;
diff --git a/Engine/Source/Runtime/IOS/MarketplaceKit/Private/MarketplaceKit.cpp b/Engine/Source/Runtime/IOS/MarketplaceKit/Private/MarketplaceKit.cpp
--- a/Engine/Source/Runtime/IOS/MarketplaceKit/Private/MarketplaceKit.cpp
+++ b/Engine/Source/Runtime/IOS/MarketplaceKit/Private/MarketplaceKit.cpp
@@ -103,6 +103,7 @@ void FMarketplaceKitModule::CacheValue()
 
 	// wait for a result, but timeout after 1s
 	dispatch_semaphore_wait(Semaphore, dispatch_time(DISPATCH_TIME_NOW, NSEC_PER_SEC));
+    dispatch_release(Semaphore);
 }
 
 IMPLEMENT_MODULE(FMarketplaceKitModule, MarketplaceKit);
