Fix for MakeSearchDictionary crash

diff --git a/Engine/Plugins/Online/WebAuth/Source/Private/IOS/IOSPlatformWebAuth.cpp b/Engine/Plugins/Online/WebAuth/Source/Private/IOS/IOSPlatformWebAuth.cpp
--- a/Engine/Plugins/Online/WebAuth/Source/Private/IOS/IOSPlatformWebAuth.cpp
+++ b/Engine/Plugins/Online/WebAuth/Source/Private/IOS/IOSPlatformWebAuth.cpp
@@ -29,7 +29,7 @@
 
 @end
 
-static NSMutableDictionary* MakeSearchDictionary(NSString *EnvironmentName);
+static NSMutableDictionary* NewSearchDictionary(NSString *EnvironmentName);
 static PresentationContext* PresentationContextProvider = nullptr;
 
 // Returns true if the request is made
@@ -97,7 +97,7 @@ bool FIOSWebAuth::AuthSessionWithURL(const FString &UrlStr, const FString &Schem
 	return bSessionInProgress;
 }
 
-NSMutableDictionary* MakeSearchDictionary(NSString *EnvironmentName)
+NSMutableDictionary* NewSearchDictionary(NSString *EnvironmentName)
 {
     static NSString* ServiceName = [[UIDevice currentDevice] identifierForVendor].UUIDString;
 
@@ -122,7 +122,7 @@ bool FIOSWebAuth::SaveCredentials(const FString& IdStr, const FString& TokenStr,
 	FTCHARToUTF8 TCEnvironmentNameStr(*EnvironmentNameStr);
 	NSString *EnvironmentName = [NSString stringWithUTF8String:TCEnvironmentNameStr.Get()];
 
-	NSMutableDictionary* SearchDictionary = MakeSearchDictionary(EnvironmentName);
+	NSMutableDictionary* SearchDictionary = NewSearchDictionary(EnvironmentName);
 
 	// erase any existing one
 	SecItemDelete((CFDictionaryRef)SearchDictionary);
@@ -141,7 +141,8 @@ bool FIOSWebAuth::SaveCredentials(const FString& IdStr, const FString& TokenStr,
 	// add it
 	OSStatus Status = SecItemAdd((CFDictionaryRef)SearchDictionary, NULL);
 	NSLog(@"Tried to add, status = %d", Status);
-
+    
+    [SearchDictionary release];
 	return Status == errSecSuccess;
 }
 
@@ -152,7 +153,7 @@ bool FIOSWebAuth::LoadCredentials(FString& OutIdStr, FString& OutTokenStr, const
 
 	NSString* Id = [NSString string];
 	NSString* Token = [NSString string];
-	NSMutableDictionary* SearchDictionary = MakeSearchDictionary(EnvironmentName);
+	NSMutableDictionary* SearchDictionary = NewSearchDictionary(EnvironmentName);
 
 	// a couple extra params for retrieval
 	[SearchDictionary setObject:(id)kSecMatchLimitOne forKey:(id)kSecMatchLimit];
@@ -203,6 +204,7 @@ bool FIOSWebAuth::LoadCredentials(FString& OutIdStr, FString& OutTokenStr, const
 		OutTokenStr = FString();
 	}
 
+    [SearchDictionary release];
 	return Status == errSecSuccess;
 }
 
