Fix for MakeSearchDictionary crash

#rnx

diff --git a/Engine/Plugins/Online/WebAuth/Source/Private/IOS/IOSPlatformWebAuth.cpp b/Engine/Plugins/Online/WebAuth/Source/Private/IOS/IOSPlatformWebAuth.cpp
--- a/Engine/Plugins/Online/WebAuth/Source/Private/IOS/IOSPlatformWebAuth.cpp
+++ b/Engine/Plugins/Online/WebAuth/Source/Private/IOS/IOSPlatformWebAuth.cpp
@@ -110,7 +110,7 @@ NSMutableDictionary* NewSearchDictionary(NSString *EnvironmentName)
 	[SearchDictionary setObject:EncodedIdentifier forKey:(id)kSecAttrAccount];
 	[SearchDictionary setObject:ServiceName forKey:(id)kSecAttrService];
 
-	return [SearchDictionary autorelease];
+	return SearchDictionary;
 }
 
 bool FIOSWebAuth::SaveCredentials(const FString& IdStr, const FString& TokenStr, const FString& EnvironmentNameStr)
