Revert to a memory leak fix, which caused log in crash

#rnx

diff --git a/Engine/Plugins/Online/WebAuth/Source/Private/IOS/IOSPlatformWebAuth.cpp b/Engine/Plugins/Online/WebAuth/Source/Private/IOS/IOSPlatformWebAuth.cpp
--- a/Engine/Plugins/Online/WebAuth/Source/Private/IOS/IOSPlatformWebAuth.cpp
+++ b/Engine/Plugins/Online/WebAuth/Source/Private/IOS/IOSPlatformWebAuth.cpp
@@ -29,7 +29,7 @@
 
 @end
 
-static NSMutableDictionary* NewSearchDictionary(NSString *EnvironmentName);
+static NSMutableDictionary* MakeSearchDictionary(NSString *EnvironmentName);
 static PresentationContext* PresentationContextProvider = nullptr;
 
 // Returns true if the request is made
@@ -92,12 +92,11 @@ bool FIOSWebAuth::AuthSessionWithURL(const FString &UrlStr, const FString &Schem
 	});
 
 	[SavedAuthSession start];
-    [SavedAuthSession release];
 
 	return bSessionInProgress;
 }
 
-NSMutableDictionary* NewSearchDictionary(NSString *EnvironmentName)
+NSMutableDictionary* MakeSearchDictionary(NSString *EnvironmentName)
 {
     static NSString* ServiceName = [[UIDevice currentDevice] identifierForVendor].UUIDString;
 
@@ -122,7 +121,7 @@ bool FIOSWebAuth::SaveCredentials(const FString& IdStr, const FString& TokenStr,
 	FTCHARToUTF8 TCEnvironmentNameStr(*EnvironmentNameStr);
 	NSString *EnvironmentName = [NSString stringWithUTF8String:TCEnvironmentNameStr.Get()];
 
-	NSMutableDictionary* SearchDictionary = NewSearchDictionary(EnvironmentName);
+	NSMutableDictionary* SearchDictionary = MakeSearchDictionary(EnvironmentName);
 
 	// erase any existing one
 	SecItemDelete((CFDictionaryRef)SearchDictionary);
@@ -141,8 +140,7 @@ bool FIOSWebAuth::SaveCredentials(const FString& IdStr, const FString& TokenStr,
 	// add it
 	OSStatus Status = SecItemAdd((CFDictionaryRef)SearchDictionary, NULL);
 	NSLog(@"Tried to add, status = %d", Status);
-    
-    [SearchDictionary release];
+
 	return Status == errSecSuccess;
 }
 
@@ -153,7 +151,7 @@ bool FIOSWebAuth::LoadCredentials(FString& OutIdStr, FString& OutTokenStr, const
 
 	NSString* Id = [NSString string];
 	NSString* Token = [NSString string];
-	NSMutableDictionary* SearchDictionary = NewSearchDictionary(EnvironmentName);
+	NSMutableDictionary* SearchDictionary = MakeSearchDictionary(EnvironmentName);
 
 	// a couple extra params for retrieval
 	[SearchDictionary setObject:(id)kSecMatchLimitOne forKey:(id)kSecMatchLimit];
@@ -204,7 +202,6 @@ bool FIOSWebAuth::LoadCredentials(FString& OutIdStr, FString& OutTokenStr, const
 		OutTokenStr = FString();
 	}
 
-    [SearchDictionary release];
 	return Status == errSecSuccess;
 }
 
