Use a string that doesn't change after reinstalls as key for keychain to store credentials for WebAuth

diff --git a/Engine/Plugins/Online/WebAuth/Source/Private/IOS/IOSPlatformWebAuth.cpp b/Engine/Plugins/Online/WebAuth/Source/Private/IOS/IOSPlatformWebAuth.cpp
--- a/Engine/Plugins/Online/WebAuth/Source/Private/IOS/IOSPlatformWebAuth.cpp
+++ b/Engine/Plugins/Online/WebAuth/Source/Private/IOS/IOSPlatformWebAuth.cpp
@@ -29,7 +29,8 @@
 
 @end
 
-static NSMutableDictionary* NewSearchDictionary(NSString *EnvironmentName);
+static NSMutableDictionary* NewSearchDictionary(NSString *EnvironmentName, bool bUseDeprecatedLocation);
+static NSData* GetStoredCredentialsData(NSString *EnvironmentName, bool bUseDeprecatedLocation, OSStatus& OutStatus);
 static PresentationContext* PresentationContextProvider = nullptr;
 
 bool FIOSWebAuth::AuthSessionWithURL(const FString &UrlStr, const FString &SchemeStr, const FWebAuthSessionCompleteDelegate& Delegate)
@@ -106,10 +107,17 @@ bool FIOSWebAuth::AuthSessionWithURL(const FString &UrlStr, const FString &Schem
 	return true;
 }
 
-NSMutableDictionary* NewSearchDictionary(NSString *EnvironmentName)
+NSMutableDictionary* NewSearchDictionary(NSString *EnvironmentName, bool bUseDeprecatedLocation)
 {
-    NSString* ServiceName = [[UIDevice currentDevice] identifierForVendor].UUIDString;
-
+	NSString* ServiceName = nil;
+	if (bUseDeprecatedLocation)
+	{
+		ServiceName = [[UIDevice currentDevice] identifierForVendor].UUIDString;
+	}
+	else
+	{
+		ServiceName = [[NSBundle mainBundle] bundleIdentifier];
+	}
 	NSString* keyName = [NSString stringWithFormat:@"DeviceCredentials_%@", EnvironmentName];
 	NSData* EncodedIdentifier = [keyName dataUsingEncoding:NSUTF8StringEncoding];
 
@@ -122,6 +130,28 @@ NSMutableDictionary* NewSearchDictionary(NSString *EnvironmentName)
 	return SearchDictionary;
 }
 
+NSData* GetStoredCredentialsData(NSString *EnvironmentName, bool bUseDeprecatedLocation, OSStatus& OutStatus)
+{
+	NSMutableDictionary* SearchDictionary = NewSearchDictionary(EnvironmentName, bUseDeprecatedLocation);
+
+	// a couple extra params for retrieval
+	[SearchDictionary setObject:(id)kSecMatchLimitOne forKey:(id)kSecMatchLimit];
+	// Add search return types
+	[SearchDictionary setObject:(id)kCFBooleanTrue forKey:(id)kSecReturnData];
+
+	CFTypeRef CFResult = nil;
+	OutStatus = SecItemCopyMatching((__bridge CFDictionaryRef)SearchDictionary, (CFTypeRef *)&CFResult);
+	[SearchDictionary release];
+
+	NSData* Result = nil;
+	if (OutStatus == errSecSuccess)
+	{
+		// only get the Results if we succeeded
+		Result = (__bridge NSData*)CFResult;
+	}
+	return Result;
+}
+
 bool FIOSWebAuth::SaveCredentials(const FString& IdStr, const FString& TokenStr, const FString& EnvironmentNameStr)
 {
 	FTCHARToUTF8 TCIdStr(*IdStr);
@@ -131,7 +161,7 @@ bool FIOSWebAuth::SaveCredentials(const FString& IdStr, const FString& TokenStr,
 	FTCHARToUTF8 TCEnvironmentNameStr(*EnvironmentNameStr);
 	NSString *EnvironmentName = [NSString stringWithUTF8String:TCEnvironmentNameStr.Get()];
 
-	NSMutableDictionary* SearchDictionary = NewSearchDictionary(EnvironmentName);
+	NSMutableDictionary* SearchDictionary = NewSearchDictionary(EnvironmentName, false);
 
 	// erase any existing one
 	SecItemDelete((CFDictionaryRef)SearchDictionary);
@@ -160,23 +190,11 @@ bool FIOSWebAuth::LoadCredentials(FString& OutIdStr, FString& OutTokenStr, const
 	FTCHARToUTF8 TCEnvironmentNameStr(*EnvironmentNameStr);
 	NSString *EnvironmentName = [NSString stringWithUTF8String:TCEnvironmentNameStr.Get()];
 
-	NSString* Id = [NSString string];
-	NSString* Token = [NSString string];
-	NSMutableDictionary* SearchDictionary = NewSearchDictionary(EnvironmentName);
-
-	// a couple extra params for retrieval
-	[SearchDictionary setObject:(id)kSecMatchLimitOne forKey:(id)kSecMatchLimit];
-	// Add search return types
-	[SearchDictionary setObject:(id)kCFBooleanTrue forKey:(id)kSecReturnData];
-
-	CFTypeRef CFResult = nil;
-	OSStatus Status = SecItemCopyMatching((__bridge CFDictionaryRef)SearchDictionary, (CFTypeRef *)&CFResult);
-
-	NSData* Result = nil;
-	if (Status == errSecSuccess)
+	OSStatus Status = errSecSuccess;
+	NSData* Result = GetStoredCredentialsData(EnvironmentName, false, Status);
+	if (Result == nil)
 	{
-		// only get the Results if we succeeded
-		Result = (__bridge NSData*)CFResult;
+		Result = GetStoredCredentialsData(EnvironmentName, true, Status);
 	}
 
 	if (Result)
@@ -185,8 +203,8 @@ bool FIOSWebAuth::LoadCredentials(FString& OutIdStr, FString& OutTokenStr, const
 		NSArray* CredentialsArray = [NSKeyedUnarchiver unarchivedObjectOfClass:[NSArray class] fromData:Result error:nil];
 		if ([CredentialsArray count] == 2 && [[CredentialsArray objectAtIndex:0] isKindOfClass:[NSString class]] && [[CredentialsArray objectAtIndex:1] isKindOfClass:[NSString class]])
 		{
-			Id = [CredentialsArray objectAtIndex:0];
-			Token = [CredentialsArray objectAtIndex:1];
+			NSString* Id = [CredentialsArray objectAtIndex:0];
+			NSString* Token = [CredentialsArray objectAtIndex:1];
 
 			const char *CStr = [Id UTF8String];
 			OutIdStr = FString(CStr);
@@ -213,7 +231,6 @@ bool FIOSWebAuth::LoadCredentials(FString& OutIdStr, FString& OutTokenStr, const
 		OutTokenStr = FString();
 	}
 
-    [SearchDictionary release];
 	return Status == errSecSuccess;
 }
 
