WebAuth (IOS)
- For loading and saving credentials, use UE Log macros instead of NS_LOG and cover more error cases
- Migrate credentials to new location to avoid stale credentials being used when deleting credentials in new location

diff --git a/Engine/Plugins/Online/WebAuth/Source/Private/IOS/IOSPlatformWebAuth.cpp b/Engine/Plugins/Online/WebAuth/Source/Private/IOS/IOSPlatformWebAuth.cpp
--- a/Engine/Plugins/Online/WebAuth/Source/Private/IOS/IOSPlatformWebAuth.cpp
+++ b/Engine/Plugins/Online/WebAuth/Source/Private/IOS/IOSPlatformWebAuth.cpp
@@ -11,6 +11,8 @@
 #include "IOS/IOSAsyncTask.h"
 #include "IOS/IOSAppDelegate.h"
 
+#include "WebAuthModule.h"
+
 @interface PresentationContext : NSObject <ASWebAuthenticationPresentationContextProviding>
 {
 }
@@ -31,6 +33,7 @@
 
 static NSMutableDictionary* NewSearchDictionary(NSString *EnvironmentName, bool bUseDeprecatedLocation);
 static NSData* GetStoredCredentialsData(NSString *EnvironmentName, bool bUseDeprecatedLocation, OSStatus& OutStatus);
+static void MigrateCredentialsData(NSString* EnvironmentName, NSData* CredentialsData);
 static PresentationContext* PresentationContextProvider = nullptr;
 
 bool FIOSWebAuth::AuthSessionWithURL(const FString &UrlStr, const FString &SchemeStr, const FWebAuthSessionCompleteDelegate& Delegate)
@@ -152,6 +155,31 @@ NSData* GetStoredCredentialsData(NSString *EnvironmentName, bool bUseDeprecatedL
 	return Result;
 }
 
+void MigrateCredentialsData(NSString* EnvironmentName, NSData* CredentialsData)
+{
+	NSMutableDictionary* NewLocationSearchDictionary = NewSearchDictionary(EnvironmentName, false);
+	
+	// Save the credentials in the new location. We do not need to clear them first as we only migrate if there are no credentials present there
+	[NewLocationSearchDictionary setObject:CredentialsData forKey:(id)kSecValueData];
+	OSStatus Status = SecItemAdd((CFDictionaryRef)NewLocationSearchDictionary, NULL);
+	
+	if (Status == errSecSuccess)
+	{
+		// Now that the credentials are successfully stored in the new location, clear the old location
+		NSMutableDictionary* OldLocationSearchDictionary = NewSearchDictionary(EnvironmentName, true);
+		SecItemDelete((CFDictionaryRef)OldLocationSearchDictionary);
+		[OldLocationSearchDictionary release];
+		
+		UE_LOG(LogWebAuth, Verbose, TEXT("Migrated credentials to new location."));
+	}
+	else
+	{
+		UE_LOG(LogWebAuth, Warning, TEXT("Failed to migrate credentials. Status = %d"), Status);
+	}
+	
+	[NewLocationSearchDictionary release];
+}
+
 bool FIOSWebAuth::SaveCredentials(const FString& IdStr, const FString& TokenStr, const FString& EnvironmentNameStr)
 {
 	FTCHARToUTF8 TCIdStr(*IdStr);
@@ -169,7 +197,7 @@ bool FIOSWebAuth::SaveCredentials(const FString& IdStr, const FString& TokenStr,
 	// if we have nil/empty params, we just delete and leave
 	if (Id == nil || Id.length == 0 || Token == nil || Token.length == 0)
 	{
-		NSLog(@"Cleared existing credentials");
+		UE_LOG(LogWebAuth, Verbose, TEXT("Cleared existing credentials"));
 		return true;
 	}
 
@@ -179,7 +207,14 @@ bool FIOSWebAuth::SaveCredentials(const FString& IdStr, const FString& TokenStr,
 
 	// add it
 	OSStatus Status = SecItemAdd((CFDictionaryRef)SearchDictionary, NULL);
-	NSLog(@"Tried to add, status = %d", Status);
+	if (Status == errSecSuccess)
+	{
+		UE_LOG(LogWebAuth, Verbose, TEXT("Saved credentials"));
+	}
+	else
+	{
+		UE_LOG(LogWebAuth, Warning, TEXT("Failed to save credentials. Status = %d"), Status);
+	}
     
     [SearchDictionary release];
 	return Status == errSecSuccess;
@@ -195,6 +230,12 @@ bool FIOSWebAuth::LoadCredentials(FString& OutIdStr, FString& OutTokenStr, const
 	if (Result == nil)
 	{
 		Result = GetStoredCredentialsData(EnvironmentName, true, Status);
+		
+		if (Result != nil)
+		{
+			// migrate data to new location
+			MigrateCredentialsData(EnvironmentName, Result);
+		}
 	}
 
 	if (Result)
@@ -212,17 +253,24 @@ bool FIOSWebAuth::LoadCredentials(FString& OutIdStr, FString& OutTokenStr, const
 			CStr = [Token UTF8String];
 			OutTokenStr = FString(CStr);
 
-			NSLog(@"Retrieved credentials successfully");
+			UE_LOG(LogWebAuth, Verbose, TEXT("Retrieved credentials successfully"));
 		}
 		else
 		{
-			NSLog(@"Retrieved credentials, but they were poorly formatted, failing.");
+			UE_LOG(LogWebAuth, Warning, TEXT("Retrieved credentials, but they were poorly formatted, failing."));
 			Status = errSecParam;
 		}
 	}
 	else
 	{
-		NSLog(@"Failed to retrieve, status = %d", Status);
+		if (Status == errSecItemNotFound)
+		{
+			UE_LOG(LogWebAuth, Log, TEXT("Failed to retrieve credentials, no credentials found"));
+		}
+		else
+		{
+			UE_LOG(LogWebAuth, Warning, TEXT("Failed to retrieve credentials, status = %d"), Status);
+		}
 	}
 
 	if (Status != errSecSuccess)
diff --git a/Engine/Plugins/Online/WebAuth/Source/Public/WebAuthModule.h b/Engine/Plugins/Online/WebAuth/Source/Public/WebAuthModule.h
--- a/Engine/Plugins/Online/WebAuth/Source/Public/WebAuthModule.h
+++ b/Engine/Plugins/Online/WebAuth/Source/Public/WebAuthModule.h
@@ -11,7 +11,7 @@
 #include "PlatformWebAuth.h"
 #include "WebAuth.h"
 
-WEBAUTH_API DECLARE_LOG_CATEGORY_EXTERN(LogWebAuth, Display, All);
+WEBAUTH_API DECLARE_LOG_CATEGORY_EXTERN(LogWebAuth, Log, All);
 
 
 /**
