Made WebAuth redirect more robust on Android

#android

diff --git a/Engine/Build/Android/Java/src/com/epicgames/unreal/GameActivity.java.template b/Engine/Build/Android/Java/src/com/epicgames/unreal/GameActivity.java.template
--- a/Engine/Build/Android/Java/src/com/epicgames/unreal/GameActivity.java.template
+++ b/Engine/Build/Android/Java/src/com/epicgames/unreal/GameActivity.java.template
@@ -5108,7 +5108,8 @@ Trace.endSection();
 		// Forcing this to false so the virtual keyboard can be shown again after resuming
 		// since calls to showSoftInput are ignored on resume so have to make sure state is reset
 		bKeyboardShowing = false;
-
+		
+		boolean bHasNewIntent = bIntentModified;
 		if (bIntentModified)
 		{
 			bIntentModified = false;
diff --git a/Engine/Plugins/Online/WebAuth/Source/WebAuth_UPL.xml b/Engine/Plugins/Online/WebAuth/Source/WebAuth_UPL.xml
--- a/Engine/Plugins/Online/WebAuth/Source/WebAuth_UPL.xml
+++ b/Engine/Plugins/Online/WebAuth/Source/WebAuth_UPL.xml
@@ -31,12 +31,12 @@ import java.net.HttpCookie;
 	if (WebAuthAwaitingResponse)
 	{
 		Intent webAuthIntent = getIntent();
-		if (ACTION_VIEW.equals(webAuthIntent.getAction()))
+		if (bHasNewIntent &amp;&amp; ACTION_VIEW.equals(webAuthIntent.getAction()))
 		{
 			Uri data = webAuthIntent.getData();
 			handleAuthenticationUri(data);
 		}
-		else
+		else if (hasWindowFocus())
 		{
 			handleAuthenticationUri(null);
 		}
@@ -100,6 +100,20 @@ import java.net.HttpCookie;
 		}
 	}
 
+	@Override
+	public void onWindowFocusChanged(boolean hasFocus)
+	{
+		super.onWindowFocusChanged(hasFocus);
+		
+		if (hasFocus && WebAuthAwaitingResponse)
+		{
+			if (!bIntentModified && !IsActivityPaused())
+			{
+				handleAuthenticationUri(null);
+			}
+		}
+	}
+	
 	public void AndroidThunkJava_WebAuthStoreCredentials(String Id, String Token, String EnvironmentName)
 	{
 		if (!Token.isEmpty() && !Id.isEmpty())
