Added AuthTab support to WebAuth plugin

#android

#Synced-CL 49977366

diff --git a/Engine/Plugins/Online/WebAuth/Source/WebAuth_UPL.xml b/Engine/Plugins/Online/WebAuth/Source/WebAuth_UPL.xml
--- a/Engine/Plugins/Online/WebAuth/Source/WebAuth_UPL.xml
+++ b/Engine/Plugins/Online/WebAuth/Source/WebAuth_UPL.xml
@@ -45,42 +45,32 @@ import java.net.HttpCookie;
 		</insert>
 	</gameActivityOnResumeAdditions>
 
-	<!-- Chrome will ALWAYS respond with canceled here, if there is a valid result, this will be returned via the intent, but this occurs at 'some unknown' later time
-	<gameActivityOnActivityResultAdditions>
-		<insert>
-	// Begin WebAuth onActivityResult
-	if (requestCode == AUTH_SESSION_REQUEST_CODE)
-	{
-		if (resultCode == RESULT_CANCELED)
-		{
-			handleAuthenticationUri(null);
-		}
-	}
-	// End WebAuth onActivityResult
-		</insert>
-	</gameActivityOnActivityResultAdditions>
-	-->
-
 	<gameActivityClassAdditions>
 		<insert>
 			<![CDATA[
-	private static final int AUTH_SESSION_REQUEST_CODE = 101;
 	private boolean WebAuthAwaitingResponse = false;
 
 	public native void handleAuthSessionResponse(String redirectURL);
 
 	public void AndroidThunkJava_StartAuthSession(String uri)
 	{
-		CustomTabsIntent.Builder builder = new CustomTabsIntent.Builder();
-		CustomTabsIntent customTabsIntent = builder.build();
-		customTabsIntent.intent.setData(Uri.parse(uri));
-		try {
+		Intent intent = new Intent(Intent.ACTION_VIEW);
+		// From androidx source code
+		intent.putExtra("androidx.browser.auth.extra.LAUNCH_AUTH_TAB", true);
+		// A null EXTRA_SESSION is interpreted as a CustomTabs intent by browsers that don't support AuthTab
+		Bundle bundle = new Bundle();
+		bundle.putBinder(CustomTabsIntent.EXTRA_SESSION, null);
+		intent.putExtras(bundle);
+		intent.setData(Uri.parse(uri));
+		try
+		{
 			WebAuthAwaitingResponse = true;
-			startActivityForResult(customTabsIntent.intent, AUTH_SESSION_REQUEST_CODE);
+			startActivity(intent);
 		}
-		catch (Exception e) {
-			// CustomTabs not supported.. for now throwing a login error is better than crashing
-			Log.debug("CustomTabs not supported");
+		catch (Throwable e)
+		{
+			// Browser not installed... for now throwing a login error is better than crashing
+			Log.debug("Browser not installed");
 			WebAuthAwaitingResponse = false;
 			handleAuthSessionResponse("");
 		}
