Auto return new intent data in WebAuth if app restarted during login session

#android

diff --git a/Engine/Plugins/Online/WebAuth/Source/WebAuth_UPL.xml b/Engine/Plugins/Online/WebAuth/Source/WebAuth_UPL.xml
--- a/Engine/Plugins/Online/WebAuth/Source/WebAuth_UPL.xml
+++ b/Engine/Plugins/Online/WebAuth/Source/WebAuth_UPL.xml
@@ -54,6 +54,27 @@ import java.net.HttpCookie;
 
 	public void AndroidThunkJava_StartAuthSession(String uri)
 	{
+		SharedPreferences preferences = getSharedPreferences("WebAuth", Context.MODE_PRIVATE);
+
+		boolean bUseWebAuthAutoUseIntentDataOnFirstTry = "true".equals(ConfigRulesVars.getOrDefault("WebAuthAutoUseIntentDataOnFirstTry", "false"));
+		if (bUseWebAuthAutoUseIntentDataOnFirstTry)
+		{ 
+			if (preferences.getBoolean("WebAuthAwaitingResponse", false))
+			{
+				Intent webAuthIntent = getIntent();
+				if (webAuthIntent != null && ACTION_VIEW.equals(webAuthIntent.getAction()))
+				{
+					Uri data = webAuthIntent.getData();
+					if (data != null)
+					{
+						preferences.edit().putBoolean("WebAuthAwaitingResponse", false).apply();
+						handleAuthSessionResponse(data.toString());
+						return;
+					}
+				}
+			}
+		}
+		
 		Intent intent = new Intent(Intent.ACTION_VIEW);
 		// From androidx source code
 		intent.putExtra("androidx.browser.auth.extra.LAUNCH_AUTH_TAB", true);
@@ -64,6 +85,7 @@ import java.net.HttpCookie;
 		intent.setData(Uri.parse(uri));
 		try
 		{
+			preferences.edit().putBoolean("WebAuthAwaitingResponse", true).apply();
 			WebAuthAwaitingResponse = true;
 			startActivity(intent);
 		}
@@ -71,6 +93,7 @@ import java.net.HttpCookie;
 		{
 			// Browser not installed... for now throwing a login error is better than crashing
 			Log.debug("Browser not installed");
+			preferences.edit().putBoolean("WebAuthAwaitingResponse", false).apply();
 			WebAuthAwaitingResponse = false;
 			handleAuthSessionResponse("");
 		}
@@ -78,6 +101,7 @@ import java.net.HttpCookie;
 
 	private void handleAuthenticationUri(Uri data)
 	{
+		getSharedPreferences("WebAuth", Context.MODE_PRIVATE).edit().putBoolean("WebAuthAwaitingResponse", false).apply();
 		WebAuthAwaitingResponse = false;
 
 		if (data != null)
